Move VPSettings and VPRoles into their own modules / plugins.

diff --git a/Engine/Plugins/Experimental/VirtualCamera/Source/VCamCore/VCamCore.Build.cs b/Engine/Plugins/Experimental/VirtualCamera/Source/VCamCore/VCamCore.Build.cs
--- a/Engine/Plugins/Experimental/VirtualCamera/Source/VCamCore/VCamCore.Build.cs
+++ b/Engine/Plugins/Experimental/VirtualCamera/Source/VCamCore/VCamCore.Build.cs
@@ -49,7 +49,8 @@ public class VCamCore : ModuleRules
 				"MediaIOCore",
 				"RemoteSession",
 				"InputCore",
-				"VPRoles"
+				"VPRoles",
+				"VPSettings"
 				// ... add private dependencies that you statically link with here ...
 			}
 			);
diff --git a/Engine/Plugins/Experimental/VirtualCamera/VirtualCamera.uplugin b/Engine/Plugins/Experimental/VirtualCamera/VirtualCamera.uplugin
--- a/Engine/Plugins/Experimental/VirtualCamera/VirtualCamera.uplugin
+++ b/Engine/Plugins/Experimental/VirtualCamera/VirtualCamera.uplugin
@@ -99,6 +99,10 @@
 			"Name": "VirtualProductionUtilities",
 			"Enabled": true
 		},
+		{
+			"Name": "VPRoles",
+			"Enabled": true
+		},
 		{
 			"Name": "LiveLink",
 			"Enabled": true
diff --git a/Engine/Plugins/Experimental/VirtualProduction/VPRoles/Source/VPRoles/Private/VPRolesModule.cpp b/Engine/Plugins/Experimental/VirtualProduction/VPRoles/Source/VPRoles/Private/VPRolesModule.cpp
new file mode 100644
--- /dev/null
+++ b/Engine/Plugins/Experimental/VirtualProduction/VPRoles/Source/VPRoles/Private/VPRolesModule.cpp
@@ -0,0 +1,10 @@
+// Copyright Epic Games, Inc. All Rights Reserved.
+
+#include "VPRolesModule.h"
+#include "CoreMinimal.h"
+#include "Modules/ModuleManager.h"
+#include "Modules/ModuleInterface.h"
+
+DEFINE_LOG_CATEGORY(LogVPRoles);
+
+IMPLEMENT_MODULE(FDefaultModuleImpl, VPRoles)
diff --git a/Engine/Plugins/Experimental/VirtualProduction/VPRoles/Source/VPRoles/Private/VPRolesModule.h b/Engine/Plugins/Experimental/VirtualProduction/VPRoles/Source/VPRoles/Private/VPRolesModule.h
new file mode 100644
--- /dev/null
+++ b/Engine/Plugins/Experimental/VirtualProduction/VPRoles/Source/VPRoles/Private/VPRolesModule.h
@@ -0,0 +1,7 @@
+// Copyright Epic Games, Inc. All Rights Reserved.
+
+#pragma once
+
+#include "Logging/LogMacros.h"
+
+DECLARE_LOG_CATEGORY_EXTERN(LogVPRoles, Log, Log);
diff --git a/Engine/Plugins/Experimental/VirtualProduction/VPRoles/Source/VPRoles/Private/VPRolesSettings.h b/Engine/Plugins/Experimental/VirtualProduction/VPRoles/Source/VPRoles/Private/VPRolesSettings.h
new file mode 100644
--- /dev/null
+++ b/Engine/Plugins/Experimental/VirtualProduction/VPRoles/Source/VPRoles/Private/VPRolesSettings.h
@@ -0,0 +1,32 @@
+// Copyright Epic Games, Inc. All Rights Reserved.
+
+#pragma once
+
+#include "CoreMinimal.h"
+#include "GameplayTagContainer.h"
+#include "VPRolesSettings.generated.h"
+
+UCLASS(config=UserVPRoles)
+class VPROLES_API UVPRolesUserSettings : public UObject
+{
+public:
+	GENERATED_BODY()
+
+private:
+	/**
+	 * The machine role(s) in a virtual production context.
+	 * @note The role may be override via the command line, "-VPRole=[Role.SubRole1|Role.SubRole2]"
+	 */
+	UPROPERTY(config)
+	FGameplayTagContainer Roles;
+	
+	/** The files that contain the available VP Roles. */
+	UPROPERTY(config)
+	TSet<FName> RoleSources;
+	
+	/** Cached available roles since the list isn't available outside of editor. */
+	UPROPERTY(config)
+	TArray<FString> CachedRoles;
+	
+	friend class UVirtualProductionRolesSubsystem;
+};
diff --git a/Engine/Plugins/Experimental/VirtualProduction/VPRoles/Source/VPRoles/Private/VPRolesSubsystem.cpp b/Engine/Plugins/Experimental/VirtualProduction/VPRoles/Source/VPRoles/Private/VPRolesSubsystem.cpp
new file mode 100644
--- /dev/null
+++ b/Engine/Plugins/Experimental/VirtualProduction/VPRoles/Source/VPRoles/Private/VPRolesSubsystem.cpp
@@ -0,0 +1,280 @@
+// Copyright Epic Games, Inc. All Rights Reserved.
+
+#include "VPRolesSubsystem.h"
+
+#include "Algo/Transform.h"
+#include "GameplayTagContainer.h"
+#include "GameplayTagsManager.h"
+#include "Misc/CommandLine.h"
+#include "VPRolesSettings.h"
+#include "VPRolesModule.h"
+#include "VPSettings.h"
+
+#if WITH_EDITOR
+#include "Editor.h"
+#include "GameplayTagsEditorModule.h"
+#endif
+
+static const FName DefaultVPRolesFileName = ("VPRoles.ini");
+
+void UVirtualProductionRolesSubsystem::Initialize(FSubsystemCollectionBase& Collection)
+{
+	CommandLineRoles = MakePimpl<FGameplayTagContainer>();
+	ImportCurrentRolesFromCommandLine();
+	ImportCurrentRolesFromDeprecatedSetting();
+
+	AddRoleSource(DefaultVPRolesFileName);
+
+	// todo: Pop up menu to migrate all sources to the VPRoles one.
+	ImportRoleSourcesFromVPSettings();
+	
+	// Provide a getter to the old system so it can return up to date roles.
+	GetMutableDefault<UVPSettings>()->RolesGetter.BindUObject(this, &UVirtualProductionRolesSubsystem::GetRolesContainer_Private);
+}
+
+void UVirtualProductionRolesSubsystem::Deinitialize()
+{
+	if (UVPSettings* VPSettings = GetMutableDefault<UVPSettings>())
+	{
+		// If this is uninitialized before the VP Settings are destroyed, provide a fallback that returns an empty container.
+		VPSettings->RolesGetter.BindStatic(&UVirtualProductionRolesSubsystem::GetRolesContainerFallback);
+	}
+}
+
+TArray<FString> UVirtualProductionRolesSubsystem::GetActiveRoles() const
+{
+	TArray<FString> Roles;
+	Algo::Transform(GetRolesContainer_Private(), Roles, [](const FGameplayTag& Tag) { return Tag.ToString(); });
+	return Roles;
+}
+
+bool UVirtualProductionRolesSubsystem::HasActiveRole(const FString& Role) const
+{
+	constexpr bool bErrorWhenNotFound = false;
+	const FGameplayTag Tag = UGameplayTagsManager::Get().RequestGameplayTag(*Role, bErrorWhenNotFound);
+	if (Tag.IsValid())
+	{
+		return GetRolesContainer_Private().HasTag(Tag);
+	}
+	return false;
+}
+
+FString UVirtualProductionRolesSubsystem::GetActiveRolesString() const
+{
+	return GetRolesContainer_Private().ToStringSimple();
+}
+
+void UVirtualProductionRolesSubsystem::SetActiveRoles(const TArray<FString>& Roles)
+{
+	FGameplayTagContainer Container;
+	
+	for (const FString& Role : Roles)
+	{
+		constexpr bool bErrorIfNotFound = false;
+		FGameplayTag Tag = FGameplayTag::RequestGameplayTag(*Role, bErrorIfNotFound);
+		if (Tag.IsValid())
+		{
+			Container.AddTag(Tag);
+		}
+		else
+		{
+			UE_LOG(LogVPRoles, Error, TEXT("Role %s doesn't exist."), *Role);
+		}
+	}
+
+	// Clear the transient roles if we're running through code while in editor.
+	GetMutableDefault<UVPRolesUserSettings>()->Roles = MoveTemp(Container);
+
+	BroadcastRolesChanged();
+}
+
+TArray<FString> UVirtualProductionRolesSubsystem::GetAllRoles() const
+{
+#if WITH_EDITOR
+	// Get all roles from our collected sources.
+	TArray<TSharedPtr<FGameplayTagNode>> Roles;
+	for (FName RoleSource : GetRoleSources())
+	{
+		UGameplayTagsManager::Get().GetAllTagsFromSource(RoleSource, Roles);
+	}
+
+	TArray<FString> ReturnedRoles;
+	ReturnedRoles.Reserve(Roles.Num());
+	Algo::TransformIf(Roles, ReturnedRoles,
+		[](const TSharedPtr<FGameplayTagNode>& Node){ return Node.IsValid(); },
+		[](const TSharedPtr<FGameplayTagNode>& Node){ return Node->GetCompleteTagString(); });
+	
+	return ReturnedRoles;
+#else
+	return GetDefault<UVPRolesUserSettings>()->CachedRoles;
+#endif
+}
+
+#if WITH_EDITOR
+bool UVirtualProductionRolesSubsystem::AddRole(const FString& RoleName)
+{
+	const bool bResult = IGameplayTagsEditorModule::Get().AddNewGameplayTagToINI(RoleName, TEXT("VPRole"), DefaultVPRolesFileName);
+	BroadcastRolesChanged();
+	return bResult;
+}
+	
+bool UVirtualProductionRolesSubsystem::RemoveRole(const FString& RoleName)
+{
+	bool bResult = false;
+	if (const TSharedPtr<FGameplayTagNode> Node = UGameplayTagsManager::Get().FindTagNode(*RoleName))
+	{
+		bResult = IGameplayTagsEditorModule::Get().DeleteTagFromINI(Node);
+		BroadcastRolesChanged();
+	}
+	return bResult;
+}
+
+#endif
+
+bool UVirtualProductionRolesSubsystem::IsUsingCommandLineRoles() const
+{
+	return bUseCommandLineRoles; 
+}
+
+bool UVirtualProductionRolesSubsystem::HasCommandLineRoles() const
+{
+	return bHasCommandLineRoles;
+}
+
+void UVirtualProductionRolesSubsystem::UseCommandLineRoles(bool bInUseCommandLineRoles)
+ {
+	if (bHasCommandLineRoles)
+	{
+		bUseCommandLineRoles = bInUseCommandLineRoles;
+		BroadcastRolesChanged();
+	}
+}
+
+const TSet<FName>& UVirtualProductionRolesSubsystem::GetRoleSources() const
+{
+	return GetDefault<UVPRolesUserSettings>()->RoleSources;
+}
+
+FName UVirtualProductionRolesSubsystem::GetDefaultRoleSource() const
+{
+	return DefaultVPRolesFileName;
+}
+
+bool UVirtualProductionRolesSubsystem::HasActiveRole(const FGameplayTag& Role) const
+{
+	return GetRolesContainer_Private().HasTag(Role);
+}
+
+void UVirtualProductionRolesSubsystem::ImportCurrentRolesFromDeprecatedSetting()
+{
+	UVPSettings* DeprecatedSettings = GetMutableDefault<UVPSettings>();
+	GetMutableDefault<UVPRolesUserSettings>()->Roles.AppendTags(DeprecatedSettings->Roles);
+	DeprecatedSettings->Roles.Reset();
+	DeprecatedSettings->SaveConfig();
+}
+
+const FGameplayTagContainer& UVirtualProductionRolesSubsystem::GetRolesContainer_Private() const
+{
+	if (bUseCommandLineRoles)
+	{
+		return *CommandLineRoles;
+	}
+		
+	return GetMutableDefault<UVPRolesUserSettings>()->Roles;
+}
+
+const FGameplayTagContainer& UVirtualProductionRolesSubsystem::GetRolesContainerFallback()
+{
+	static const FGameplayTagContainer FallbackGetterContainer;
+	return FallbackGetterContainer;
+}
+
+FGameplayTagContainer* UVirtualProductionRolesSubsystem::GetRolesContainerPtr(EGetRolesPtrSource Source)
+{
+	if (Source == EGetRolesPtrSource::Settings)
+	{
+		return &GetMutableDefault<UVPRolesUserSettings>()->Roles;
+	}
+	else
+	{
+		return CommandLineRoles.Get();
+	}
+}
+
+void UVirtualProductionRolesSubsystem::ImportRoleSourcesFromVPSettings()
+{
+#if WITH_EDITOR
+	const UVPSettings* VPSettings = GetDefault<UVPSettings>();
+	for (const FGameplayTag& Tag : VPSettings->Roles)
+	{
+		FString DevComment;
+		FName TagSource;
+		bool bIsTagExplicit;
+		bool bIsRestrictedTag;
+		bool bAllowNonRestrictedChildren;
+		if (UGameplayTagsManager::Get().GetTagEditorData(Tag.GetTagName(), DevComment, TagSource, bIsTagExplicit, bIsRestrictedTag, bAllowNonRestrictedChildren))
+		{
+			AddRoleSource(TagSource);
+		}
+	}
+#endif
+}
+
+bool UVirtualProductionRolesSubsystem::ImportCurrentRolesFromCommandLine()
+{
+	FString ComandlineRoles;
+
+	if (FParse::Value(FCommandLine::Get(), TEXT("-VPRole="), ComandlineRoles))
+	{
+		bHasCommandLineRoles = true;
+		bUseCommandLineRoles = true;
+		
+		TArray<FString> RoleList;
+		ComandlineRoles.ParseIntoArray(RoleList, TEXT("|"), true);
+
+		for (const FString& Role : RoleList)
+		{
+			FGameplayTag Tag = FGameplayTag::RequestGameplayTag(*Role, false);
+			if (Tag.IsValid())
+			{
+				CommandLineRoles->AddTag(Tag);
+			}
+			else
+			{
+				UE_LOG(LogVPRoles, Error, TEXT("Role %s doesn't exist."), *Role);
+			}
+		}
+		return true;
+	}
+	return false;
+}
+
+void UVirtualProductionRolesSubsystem::BroadcastRolesChanged() const
+{
+#if WITH_EDITOR
+	GetMutableDefault<UVPRolesUserSettings>()->CachedRoles = GetAllRoles();
+	
+	if (GEditor)
+	{
+		// Only save out the active roles if running in editor.
+		GetMutableDefault<UVPRolesUserSettings>()->SaveConfig();
+	}
+#endif
+
+	const TArray<FString> Roles = GetActiveRoles();
+	OnRolesChangedNative.Broadcast(Roles);
+	OnRolesChangedBP.Broadcast(Roles);
+
+#if WITH_EDITOR
+PRAGMA_DISABLE_DEPRECATION_WARNINGS
+	// Trigger a broadcast on the old system.
+	GetMutableDefault<UVPSettings>()->OnRolesChanged.Broadcast();
+PRAGMA_ENABLE_DEPRECATION_WARNINGS
+#endif
+}
+
+void UVirtualProductionRolesSubsystem::AddRoleSource(FName TagSource)
+{	
+	GetMutableDefault<UVPRolesUserSettings>()->RoleSources.Add(TagSource);
+}
+
diff --git a/Engine/Plugins/Experimental/VirtualProduction/VPRoles/Source/VPRoles/Public/VPRolesSubsystem.h b/Engine/Plugins/Experimental/VirtualProduction/VPRoles/Source/VPRoles/Public/VPRolesSubsystem.h
new file mode 100644
--- /dev/null
+++ b/Engine/Plugins/Experimental/VirtualProduction/VPRoles/Source/VPRoles/Public/VPRolesSubsystem.h
@@ -0,0 +1,172 @@
+// Copyright Epic Games, Inc. All Rights Reserved.
+
+#pragma once
+
+#include "CoreMinimal.h"
+#include "Subsystems/EngineSubsystem.h"
+#include "Templates/PimplPtr.h"
+
+#include "VPRolesSubsystem.generated.h"
+
+struct FGameplayTag;
+struct FGameplayTagContainer;
+
+/**
+ * Subsystem for common Virtual Production roles operations.
+ * The machine role(s) in a virtual production context read from the command line.
+ * ie. "-VPRole=[Role.SubRole1|Role.SubRole2]"
+ */
+UCLASS()
+class VPROLES_API UVirtualProductionRolesSubsystem : public UEngineSubsystem
+{
+	GENERATED_BODY()
+
+public:
+	DECLARE_DYNAMIC_MULTICAST_DELEGATE_OneParam(FOnRolesChanged, const TArray<FString>&, EnabledRoles);
+	DECLARE_MULTICAST_DELEGATE_OneParam(FOnRolesChangedNative, const TArray<FString>&);
+	
+public:
+	//~ Begin USubsystem interface
+	virtual void Initialize(FSubsystemCollectionBase& Collection) override;
+	virtual void Deinitialize() override;
+	//~ End USubsystem interface 
+
+	/**
+	 * Get the currently active Virtual Production roles.
+	 * @return The names of the current virtual production roles.
+	 */
+	UFUNCTION(BlueprintCallable, Category = "Virtual Production|Roles")
+	TArray<FString> GetActiveRoles() const;
+
+	/**
+	 * @return Whether a role is currently active.
+	 */
+	UFUNCTION(BlueprintCallable, Category = "Virtual Production|Roles")
+	bool HasActiveRole(const FString& Role) const;
+
+	/**
+	 * @return the currently active Virtual Production roles as a comma separated string representation.
+	 */
+	UFUNCTION(BlueprintCallable, Category = "Virtual Production|Roles")
+	FString GetActiveRolesString() const;
+
+	/**
+	 * Set the current Virtual Production roles.
+	 * @param Roles the list of roles to set as active.
+	 * @note The roles must already have been added either through the UI or through AddRole.
+	 * @note This will clear the previous current roles.
+	 */
+	UFUNCTION(BlueprintCallable, Category = "Virtual Production|Roles")
+	void SetActiveRoles(const TArray<FString>& Roles);
+
+	/**
+	 * Get all available roles that can be set as current.
+	 * @return The available roles.
+	 */
+	UFUNCTION(BlueprintCallable, Category = "Virtual Production|Roles")
+	TArray<FString> GetAllRoles() const;
+	
+#if WITH_EDITOR
+	/**
+	 * Add a new virtual production role.
+	 * @param RoleName the name of the role to add.
+	 * @return Whether the operation succeeded.
+	 * @note This will attempt to modify the underlying VProles.ini file so it must be checked out or made writable or the operation will fail.
+	 */
+	UFUNCTION(BlueprintCallable, Category = "Virtual Production|Roles")
+	bool AddRole(const FString& RoleName);
+
+	/**
+	 * Remove a virtual production role.
+	 * @param RoleName the name of the role to remove.
+	 * @return Whether the operation succeeded.
+	 * @note This will attempt to modify the underlying VProles.ini file so it must be checked out or made writable or the operation will fail.
+	 */
+	UFUNCTION(BlueprintCallable, Category = "Virtual Production|Roles")
+	bool RemoveRole(const FString& RoleName);
+#endif
+
+	/**
+	 * Delegate called when the roles are modified, ie. role is added, current roles are replaced, etc.
+	 */
+	FOnRolesChangedNative& OnRolesChanged()
+	{
+		return OnRolesChangedNative;
+	}
+
+	/**
+	 * @return Whether we are using the roles passed in command line if any.  
+	 */
+	bool IsUsingCommandLineRoles() const;
+
+	/**
+	 * @return Whether roles were passed through the command line.
+	 */
+	bool HasCommandLineRoles() const;
+
+	/**
+	 * Set whether the command line roles should be used if any were specified.
+	 */
+	void UseCommandLineRoles(bool bUseCommandLineRoles);
+
+	/**
+	 * @return The list of sources that provide VP roles.
+	 */
+	const TSet<FName>& GetRoleSources() const;
+	
+	/**
+	 * @return The name of the default VP role ini file.
+	 */
+	FName GetDefaultRoleSource() const;
+
+	/**
+	 * @return Whether a role is currently active.
+	 */
+	bool HasActiveRole(const FGameplayTag& Role) const;
+	
+	/**
+	 * Get the roles container, either from the command line or from the VP Roles settings.
+	 * @note New code should not use this method, it is only provided for backwards compatibility.
+	 */
+	const FGameplayTagContainer& GetRolesContainer_Private() const;
+public:
+	/**
+	 * Delegate called when the roles are modified, ie. role is added, current roles are replaced, etc.
+	 */
+	UPROPERTY(BlueprintAssignable, Category = "Virtual Production|Roles", meta = (ScriptName ="OnRolesChanged"))
+	FOnRolesChanged OnRolesChangedBP;
+	
+private:
+	/** Save roles to the VP Role config file and trigger a broadcast on the OnRolesChanged delegates. */
+	void BroadcastRolesChanged() const;
+	/** Add a role source to the VP Roles settings. Used to keep track of which config files provided VP Roles. */
+	void AddRoleSource(FName TagSource);
+	/** Parses command line roles. */
+	bool ImportCurrentRolesFromCommandLine();
+	/** Imports VP Roles from the VP Settings. */
+	void ImportCurrentRolesFromDeprecatedSetting();
+	/** Fallback handler that returns an empty tag container. */
+	static const FGameplayTagContainer& GetRolesContainerFallback();
+	enum class EGetRolesPtrSource
+	{
+		Settings,
+		CommandLine
+	};
+	/** Get a pointer to the roles gameplay tag container */
+	FGameplayTagContainer* GetRolesContainerPtr(EGetRolesPtrSource Source);
+	/** Gather all gameplay tag sources from the tags that were imported from the VP settings. */
+	void ImportRoleSourcesFromVPSettings();
+
+private:
+	/** Delegate called when the roles are modified, ie. role is added, current roles are replaced, etc. */
+	FOnRolesChangedNative OnRolesChangedNative;
+	/** Holds the role container populated from the roles specified on the command line. */
+	TPimplPtr<FGameplayTagContainer> CommandLineRoles;
+	/** Whether command line roles should be used. */
+	bool bUseCommandLineRoles = false;
+	/** Whether command line roles were specified.  */
+	bool bHasCommandLineRoles = false;
+	
+	friend class FVPRolesModule;
+	friend class FVPRolesEditorModule;
+};
diff --git a/Engine/Plugins/Experimental/VirtualProduction/VPRoles/Source/VPRoles/VPRoles.Build.cs b/Engine/Plugins/Experimental/VirtualProduction/VPRoles/Source/VPRoles/VPRoles.Build.cs
new file mode 100644
--- /dev/null
+++ b/Engine/Plugins/Experimental/VirtualProduction/VPRoles/Source/VPRoles/VPRoles.Build.cs
@@ -0,0 +1,37 @@
+// Copyright Epic Games, Inc. All Rights Reserved.
+
+using UnrealBuildTool;
+
+public class VPRoles : ModuleRules
+{
+	public VPRoles(ReadOnlyTargetRules Target) : base(Target)
+	{
+		PublicDependencyModuleNames.AddRange(
+			new string[]
+			{
+				"Core",
+				"CoreUObject",
+				"Engine",
+				"GameplayTags",
+			}
+		);
+
+		PrivateDependencyModuleNames.AddRange(
+			new string[]
+			{
+				"VPSettings"
+			}
+		);
+
+		if (Target.bBuildEditor)
+		{
+			PrivateDependencyModuleNames.AddRange(
+				new string[]
+				{
+					"GameplayTagsEditor",
+					"UnrealEd",
+				}
+			);
+		}
+	}
+}
diff --git a/Engine/Plugins/Experimental/VirtualProduction/VPRoles/Source/VPRolesEditor/Private/VPRolesEditorModule.cpp b/Engine/Plugins/Experimental/VirtualProduction/VPRoles/Source/VPRolesEditor/Private/VPRolesEditorModule.cpp
new file mode 100644
--- /dev/null
+++ b/Engine/Plugins/Experimental/VirtualProduction/VPRoles/Source/VPRolesEditor/Private/VPRolesEditorModule.cpp
@@ -0,0 +1,416 @@
+// Copyright Epic Games, Inc. All Rights Reserved.
+
+#include "VPRolesEditorModule.h"
+
+#include "Framework/Application/SlateApplication.h"
+#include "GameplayTagContainer.h"
+#include "GameplayTagsEditorModule.h"
+#include "GameplayTagsSettings.h"
+#include "LevelEditor.h"
+#include "HAL/PlatformFileManager.h"
+#include "Misc/MessageDialog.h"
+#include "SGameplayTagWidget.h"
+#include "SSettingsEditorCheckoutNotice.h"
+#include "ToolMenus.h"
+#include "ToolMenuSection.h"
+#include "VPRolesEditorStyle.h"
+#include "VPRolesSubsystem.h"
+#include "Widgets/Input/SCheckBox.h"
+#include "Widgets/SNullWidget.h"
+#include "Widgets/Input/SEditableTextBox.h"
+#include "Widgets/Layout/SWidgetSwitcher.h"
+
+
+DEFINE_LOG_CATEGORY(LogVPRolesEditor)
+#define LOCTEXT_NAMESPACE "VPRolesEditor"
+
+namespace UE::Private::VPRolesEditor
+{
+	SGameplayTagWidget::ETagFilterResult OnFilterTag(const TSharedPtr<FGameplayTagNode>& InGameplayTagNode)
+	{
+		if (GEngine)
+		{
+			const TSet<FName>& TagSources = GEngine->GetEngineSubsystem<UVirtualProductionRolesSubsystem>()->GetRoleSources();
+			
+			FString DevComment;
+			FName TagSource;
+			bool bIsExplicit = false;
+			bool bIsRestricted = false;
+			bool bAllowNonRestrictedChildren = false;
+
+			static const FString VPRoleComment = TEXT("VPRole");
+			
+			UGameplayTagsManager::Get().GetTagEditorData(InGameplayTagNode->GetCompleteTagName(), DevComment, TagSource, bIsExplicit, bIsRestricted, bAllowNonRestrictedChildren);
+			
+			if ((TagSources.Contains(TagSource) || DevComment == VPRoleComment))
+			{
+				return SGameplayTagWidget::ETagFilterResult::IncludeTag;
+			}
+		}
+		
+		return SGameplayTagWidget::ETagFilterResult::ExcludeTag;
+	};
+	
+}
+
+void FVPRolesEditorModule::StartupModule()
+{
+	FVPRolesEditorStyle::Initialize();
+	FCoreDelegates::OnPostEngineInit.AddRaw(this, &FVPRolesEditorModule::ExtendLevelEditorToolbar);
+}
+
+void FVPRolesEditorModule::ShutdownModule()
+{
+	FCoreDelegates::OnPostEngineInit.RemoveAll(this);
+	FVPRolesEditorStyle::Shutdown();
+}
+
+void FVPRolesEditorModule::ExtendLevelEditorToolbar()
+{
+	if (UVirtualProductionRolesSubsystem* RolesSubsystem = GEngine->GetEngineSubsystem<UVirtualProductionRolesSubsystem>())
+	{
+		if (UGameplayTagsSettings* GameplayTagsSettings = GetMutableDefault<UGameplayTagsSettings>())
+		{
+			GameplayTagsSettings->ImportTagsFromConfig = true;
+		}
+		
+		UToolMenu* Menu = UToolMenus::Get()->ExtendMenu("LevelEditor.LevelEditorToolBar.User");
+
+		FToolMenuEntry VPRolesEntry = FToolMenuEntry::InitComboButton(
+			"VPRolesMenu",
+			FUIAction(),
+			FOnGetContent::CreateRaw(this, &FVPRolesEditorModule::GenerateVPRolesLevelEditorToolbarMenu),
+			LOCTEXT("LevelEditorToolbarVPRolesButtonLabel", "VP Roles"),
+			LOCTEXT("LevelEditorToolbarVPRolesButtonTooltip", "Edit VP Roles"),
+			FSlateIcon(FAppStyle::GetAppStyleSetName(), "Icons.Role"));
+		
+		VPRolesEntry.StyleNameOverride = "CalloutToolbar";
+		Menu->FindOrAddSection("VPRoles").AddEntry(VPRolesEntry);
+	}
+}
+
+TSharedRef<SWidget> FVPRolesEditorModule::GenerateVPRolesLevelEditorToolbarMenu()
+{
+	bShowAddRoleTextBox = false;
+	UVirtualProductionRolesSubsystem* RolesSubsystem = GEngine->GetEngineSubsystem<UVirtualProductionRolesSubsystem>();
+	if (!RolesSubsystem)
+	{
+		return SNullWidget::NullWidget;
+	}
+
+	FLevelEditorModule& LevelEditorModule = FModuleManager::LoadModuleChecked<FLevelEditorModule>("LevelEditor");
+	TSharedPtr<FUICommandList> CommandBindings = LevelEditorModule.GetGlobalLevelEditorActions();
+
+	constexpr bool bShouldCloseMenuAfterSelection = false;
+	FMenuBuilder MenuBuilder(bShouldCloseMenuAfterSelection, CommandBindings);
+
+	constexpr bool bNoIndent = false;
+	constexpr bool bSearchable = false;
+	static const FName NoExtensionHook = NAME_None;
+	
+	{
+		MenuBuilder.BeginSection(NoExtensionHook, LOCTEXT("ModifyRolesSection", "Modify Roles"));
+		FMenuEntryParams MenuEntryParams;
+
+		FUIAction AddRoleAction;
+		AddRoleAction.ExecuteAction = FExecuteAction::CreateLambda([this]()
+		{
+			if (TSharedPtr<SWidget> AddRoleWidget = AddRoleEditableBoxWeakPtr.Pin())
+			{
+				bShowAddRoleTextBox = true;
+				FSlateApplication::Get().SetKeyboardFocus(AddRoleWidget);
+			}
+			return;
+		});
+		
+		AddRoleAction.CanExecuteAction = FCanExecuteAction::CreateLambda([this, RolesSubsystem]()
+		{
+			return !bShowAddRoleTextBox && !RolesSubsystem->IsUsingCommandLineRoles();
+		});
+
+		MenuBuilder.AddMenuEntry(LOCTEXT("AddRoleEntry", "Add Role"),
+			LOCTEXT("AddRoleToolTip", "Add a new virtual production role."),
+			FSlateIcon(FVPRolesEditorStyle::Get().GetStyleSetName(), "VPRolesEditor.AddRole"),
+			AddRoleAction);
+		
+		MenuBuilder.EndSection();
+	}
+	
+	{
+		MenuBuilder.BeginSection(NoExtensionHook, LOCTEXT("RolesSection", "Roles"));
+		MenuBuilder.AddWidget(GenerateGameplayTagWidget(), FText::GetEmpty(), bNoIndent, bSearchable);
+		MenuBuilder.EndSection();
+	}
+	
+	return MenuBuilder.MakeWidget();
+}
+
+TSharedRef<SWidget> FVPRolesEditorModule::GenerateGameplayTagWidget()
+{
+	UVirtualProductionRolesSubsystem* RolesSubsystem = GEngine->GetEngineSubsystem<UVirtualProductionRolesSubsystem>();
+	if (!RolesSubsystem)
+	{
+		return SNullWidget::NullWidget;
+	}
+	
+	if (DefaultTagSourceConfigName.IsEmpty())
+	{
+		DefaultTagSourceConfigName = RolesSubsystem->GetDefaultRoleSource().ToString();
+	}
+	
+	static FSlateColorBrush TransparentBrush{ FLinearColor::Transparent };
+	
+	TArray<SGameplayTagWidget::FEditableGameplayTagContainerDatum> EditableContainers;
+	FGameplayTagContainer* ContainerPtr = RolesSubsystem->GetRolesContainerPtr(UVirtualProductionRolesSubsystem::EGetRolesPtrSource::Settings);
+	EditableContainers.Emplace(nullptr, ContainerPtr);
+
+	TArray<SGameplayTagWidget::FEditableGameplayTagContainerDatum> CommandLineEditableContainers;
+    FGameplayTagContainer* CommandLineContainerPtr = RolesSubsystem->GetRolesContainerPtr(UVirtualProductionRolesSubsystem::EGetRolesPtrSource::CommandLine);
+    CommandLineEditableContainers.Emplace(nullptr, CommandLineContainerPtr);
+
+	return SNew(SBorder)
+		.BorderImage(FAppStyle::GetBrush("NoBorder"))
+		.Padding(FMargin(8.f, 0.f))
+		[
+			SNew(SVerticalBox)
+			+ SVerticalBox::Slot()
+			.Padding(2.f)
+			.AutoHeight()
+			[
+				SNew(SBorder)
+				.Visibility_Lambda([RolesSubsystem]() { return RolesSubsystem->HasCommandLineRoles() ? EVisibility::Visible : EVisibility::Collapsed; })
+				.Padding(FMargin(6.f, 4.f))
+				.BorderImage(FAppStyle::GetBrush("ToolPanel.GroupBorder"))
+				[
+					SNew(SHorizontalBox)
+					+ SHorizontalBox::Slot()
+					.AutoWidth()
+					.VAlign(VAlign_Center)
+					[
+						SNew(SCheckBox)
+						.IsChecked_Lambda([RolesSubsystem](){ return RolesSubsystem->IsUsingCommandLineRoles() ? ECheckBoxState::Checked : ECheckBoxState::Unchecked; })
+						.OnCheckStateChanged_Lambda([RolesSubsystem, this](ECheckBoxState CheckBoxState)
+						{
+							RolesSubsystem->UseCommandLineRoles(CheckBoxState == ECheckBoxState::Checked ? true : false);
+						})
+					]
+					+ SHorizontalBox::Slot()
+					.FillWidth(1.f)
+					.VAlign(VAlign_Center)
+					[
+						SNew(STextBlock)	
+						.Text(LOCTEXT("UseCommandLineRolesLabel", "Use command line roles"))
+					]
+				]
+			]
+			+ SVerticalBox::Slot()
+			.AutoHeight()
+			[
+				SNew(SHorizontalBox)
+				.Visibility_Lambda(	[RolesSubsystem, this](){ return !bShowAddRoleTextBox || RolesSubsystem->IsUsingCommandLineRoles() ? EVisibility::Collapsed : EVisibility::Visible; })
+				+SHorizontalBox::Slot()
+				.AutoWidth()
+				[
+					// To match the padding of the gameplay tag list.
+					SNew(SImage)
+					.Image(FAppStyle::GetBrush("TreeArrow_Collapsed"))
+					.Visibility(EVisibility::Hidden)
+				]
+				+SHorizontalBox::Slot()
+	            .AutoWidth()
+	            [
+	                SNew(SCheckBox)
+	                .IsEnabled(false)
+	            ]
+				+SHorizontalBox::Slot()
+				.Padding(0.f)
+				[
+					SAssignNew(AddRoleEditableBoxWeakPtr, SEditableTextBox)
+						.IsReadOnly_Lambda([RolesSubsystem](){ return RolesSubsystem->IsUsingCommandLineRoles(); })
+						.ClearKeyboardFocusOnCommit(false)
+						.HintText(LOCTEXT("GameplayTagWidget_AddRole", "Add Role"))
+						.OnTextChanged_Raw(this, &FVPRolesEditorModule::OnAddTagTextChanged)
+						.OnTextCommitted_Raw( this, &FVPRolesEditorModule::OnSubmitNewTag)
+				]
+			]
+			+ SVerticalBox::Slot()
+			.Padding(0.f)
+			.AutoHeight()
+			[
+				SNew(SBox)
+				.MinDesiredWidth(200.f)
+				[
+					SNew(SWidgetSwitcher)
+					.WidgetIndex_Lambda([RolesSubsystem]()
+					{
+						return RolesSubsystem->IsUsingCommandLineRoles() ? 1 : 0;
+					})
+					+SWidgetSwitcher::Slot()
+					[
+						SNew(SGameplayTagWidget, EditableContainers)
+							.Padding(0.f)
+							.OnFilterTag_Static(&UE::Private::VPRolesEditor::OnFilterTag)
+							.ForceHideAddNewTag(true)
+							.ForceHideAddNewTagSource(true)
+							.ForceHideTagTreeControls(true)
+							.GameplayTagUIMode(EGameplayTagUIMode::HybridMode)
+							.MaxHeight(0) // Height of 0 means the height is not limited.
+							.BackgroundBrush(&TransparentBrush)
+							.TagTreeViewBackgroundBrush(&TransparentBrush)
+							.OnTagChanged_Lambda([RolesSubsystem]{ RolesSubsystem->BroadcastRolesChanged(); })
+					]
+					+SWidgetSwitcher::Slot()
+					[
+						SNew(SGameplayTagWidget, CommandLineEditableContainers)
+							.Padding(0.f)
+							.OnFilterTag_Static(&UE::Private::VPRolesEditor::OnFilterTag)
+							.ForceHideAddNewTag(true)
+							.ForceHideAddNewTagSource(true)
+							.ForceHideTagTreeControls(true)
+							.GameplayTagUIMode(EGameplayTagUIMode::SelectionMode)
+							.MaxHeight(0) // Height of 0 means the height is not limited.
+							.BackgroundBrush(&TransparentBrush)
+							.TagTreeViewBackgroundBrush(&TransparentBrush)
+							.ReadOnly(true)
+					]					
+				]
+			]
+			+SVerticalBox::Slot()
+			.AutoHeight()
+			[
+				SNew(SSettingsEditorCheckoutNotice)
+				.Visibility_Raw(this, &FVPRolesEditorModule::DetermineFileWatcherWidgetVisibility)
+				.ConfigFilePath(RolesSubsystem->GetDefaultRoleSource().ToString())
+			]
+		];
+}
+
+void FVPRolesEditorModule::OnAddTagTextChanged(const FText& Text)
+{
+	if (const TSharedPtr<SEditableTextBox> AddTagEditableBox = AddRoleEditableBoxWeakPtr.Pin())
+	{
+		const FString TagName = Text.ToString();
+		FText ErrorMsg;
+
+		if (Text.IsEmpty())
+		{
+			AddTagEditableBox->SetError(FText::GetEmpty());
+			return;
+		}
+
+		if (!UGameplayTagsManager::Get().IsValidGameplayTagString(TagName, &ErrorMsg))
+		{
+			FText MessageTitle(LOCTEXT("InvalidTag", "Invalid Tag"));
+			AddTagEditableBox->SetError(ErrorMsg);
+		}
+		else
+		{
+			AddTagEditableBox->SetError(FText::GetEmpty());
+		}
+	}
+}
+
+EVisibility FVPRolesEditorModule::DetermineFileWatcherWidgetVisibility() const
+{
+	return CanModifyTagSource() ? EVisibility::Collapsed : EVisibility::Visible;
+}
+
+bool FVPRolesEditorModule::CanModifyTagSource() const
+{
+	if (!DefaultTagSourceConfigName.IsEmpty())
+	{
+		return !FPlatformFileManager::Get().GetPlatformFile().IsReadOnly(*DefaultTagSourceConfigName);
+	}
+	else
+	{
+		ensureMsgf(false, TEXT("File %s should already exist."), *DefaultTagSourceConfigName);
+		return false;
+	}
+}
+
+void FVPRolesEditorModule::OnSubmitNewTag(const FText& CommittedText, ETextCommit::Type CommitType)
+{
+	if (!GEngine)
+	{
+		return;
+	}
+	
+	const TSharedPtr<SEditableTextBox> AddTagEditableBox = AddRoleEditableBoxWeakPtr.Pin();
+	if (!AddTagEditableBox)
+	{
+		return;
+	}
+	
+	if (const UVirtualProductionRolesSubsystem* RolesSubsystem = GEngine->GetEngineSubsystem<UVirtualProductionRolesSubsystem>())
+	{
+		if (CommitType == ETextCommit::Type::OnEnter)
+		{
+			UGameplayTagsManager& Manager = UGameplayTagsManager::Get();
+
+			// Only support adding tags via ini file
+			if (Manager.ShouldImportTagsFromINI() == false)
+			{
+				return;
+			}
+
+			const FText TagNameAsText = CommittedText;
+			const FString TagName = TagNameAsText.ToString();
+
+			if (TagName.IsEmpty())
+			{
+				return;
+			}
+
+			if (DefaultTagSourceConfigName.IsEmpty())
+			{
+				return;
+			}
+
+			// Check to see if this is a valid tag
+			// First check the base rules for all tags then look for any additional rules in the delegate
+			FText ErrorMsg;
+			if (!Manager.IsValidGameplayTagString(TagName, &ErrorMsg))
+			{
+				AddTagEditableBox->SetError(ErrorMsg);
+				return;
+			}
+
+			bool bCanAddTag = false;
+			if (!CanModifyTagSource())
+			{
+				if (FMessageDialog::Open(EAppMsgType::YesNo, LOCTEXT("SaveAsDefaultsIsReadOnlyMessage", "The ini file for these settings is not currently writable. Would you like to make it writable?")) == EAppReturnType::Yes)
+				{
+					bCanAddTag = SettingsHelpers::MakeWritable(DefaultTagSourceConfigName, true);
+				}
+			}
+			else
+			{
+				bCanAddTag = true;
+				AddTagEditableBox->SetText(FText::GetEmpty());
+			}
+
+			if (bCanAddTag)
+			{
+				if (IGameplayTagsEditorModule::Get().AddNewGameplayTagToINI(TagName, FString(), *DefaultTagSourceConfigName))
+				{
+					if (TSharedPtr<FGameplayTagNode> Node = Manager.FindTagNode(*TagName))
+					{
+						RolesSubsystem->BroadcastRolesChanged();
+					}
+				}
+				
+				bShowAddRoleTextBox = false;
+			}
+		}
+
+		if (CommitType == ETextCommit::Type::OnUserMovedFocus)
+		{
+			bShowAddRoleTextBox = false;
+			AddTagEditableBox->SetText(FText::GetEmpty());
+		}
+	}
+}
+
+IMPLEMENT_MODULE(FVPRolesEditorModule, VPRolesEditor)
+
+#undef LOCTEXT_NAMESPACE /*VPRolesEditor*/
diff --git a/Engine/Plugins/Experimental/VirtualProduction/VPRoles/Source/VPRolesEditor/Private/VPRolesEditorStyle.cpp b/Engine/Plugins/Experimental/VirtualProduction/VPRoles/Source/VPRolesEditor/Private/VPRolesEditorStyle.cpp
new file mode 100644
--- /dev/null
+++ b/Engine/Plugins/Experimental/VirtualProduction/VPRoles/Source/VPRolesEditor/Private/VPRolesEditorStyle.cpp
@@ -0,0 +1,59 @@
+// Copyright Epic Games, Inc. All Rights Reserved.
+
+#include "VPRolesEditorStyle.h"
+
+#include "Brushes/SlateImageBrush.h"
+#include "Interfaces/IPluginManager.h"
+#include "Framework/Application/SlateApplication.h"
+#include "Slate/SlateGameResources.h"
+#include "Styling/SlateStyleRegistry.h"
+#include "Styling/StyleColors.h"
+
+TSharedPtr<FSlateStyleSet> FVPRolesEditorStyle::StyleInstance;
+
+void FVPRolesEditorStyle::Initialize()
+{
+	if (!StyleInstance.IsValid())
+	{
+		StyleInstance = Create();
+		FSlateStyleRegistry::RegisterSlateStyle(*StyleInstance);
+	}
+}
+
+void FVPRolesEditorStyle::Shutdown()
+{
+	FSlateStyleRegistry::UnRegisterSlateStyle(*StyleInstance);
+	ensure(StyleInstance.IsUnique());
+	StyleInstance.Reset();
+}
+
+FName FVPRolesEditorStyle::GetStyleSetName()
+{
+	static FName StyleSetName(TEXT("VPRolesEditorStyle"));
+	return StyleSetName;
+}
+
+#define IMAGE_BRUSH_SVG( RelativePath, ... ) FSlateVectorImageBrush( Style->RootToContentDir( RelativePath, TEXT(".svg") ), __VA_ARGS__ )
+
+const FVector2D Icon16x16(16.0f, 16.0f);
+const FVector2D Icon40x40(40.0f, 40.0f);
+
+TSharedRef<FSlateStyleSet> FVPRolesEditorStyle::Create()
+{
+	TSharedRef<FSlateStyleSet> Style = MakeShared<FSlateStyleSet>(FVPRolesEditorStyle::GetStyleSetName());
+	Style->SetContentRoot(IPluginManager::Get().FindPlugin("VirtualProductionUtilities")->GetBaseDir() / TEXT("Resources"));
+
+	FSlateBrush AddRoleBrush = *FAppStyle::Get().GetBrush("Icons.Plus");
+	AddRoleBrush.TintColor = FStyleColors::AccentGreen;
+
+	Style->Set("VPRolesEditor.AddRole", new FSlateBrush(AddRoleBrush));
+	return Style;
+}
+
+const ISlateStyle& FVPRolesEditorStyle::Get()
+{
+	static const FVPRolesEditorStyle Inst;
+	return *StyleInstance;
+}
+
+#undef IMAGE_BRUSH_SVG
diff --git a/Engine/Plugins/Experimental/VirtualProduction/VPRoles/Source/VPRolesEditor/Private/VPRolesEditorStyle.h b/Engine/Plugins/Experimental/VirtualProduction/VPRoles/Source/VPRolesEditor/Private/VPRolesEditorStyle.h
new file mode 100644
--- /dev/null
+++ b/Engine/Plugins/Experimental/VirtualProduction/VPRoles/Source/VPRolesEditor/Private/VPRolesEditorStyle.h
@@ -0,0 +1,27 @@
+// Copyright Epic Games, Inc. All Rights Reserved.
+
+#pragma once
+
+#include "CoreMinimal.h"
+#include "Styling/SlateStyle.h"
+
+class FVPRolesEditorStyle
+{
+public:
+	/** Initialize the style. */
+	static void Initialize();
+	/** Unregister the style. */
+	static void Shutdown();
+	/** The Slate style set for VP Roles */
+	static const ISlateStyle& Get();
+	/** Get the style set name. */
+	static FName GetStyleSetName();
+
+private:
+	/** Initializes the singleton and sets the buttons used by the vp roles widget.  */
+	static TSharedRef<class FSlateStyleSet> Create();
+
+private:
+	/** Singleton instance of the style. */
+	static TSharedPtr<class FSlateStyleSet> StyleInstance;
+};
\ No newline at end of file
diff --git a/Engine/Plugins/Experimental/VirtualProduction/VPRoles/Source/VPRolesEditor/Public/VPRolesEditorModule.h b/Engine/Plugins/Experimental/VirtualProduction/VPRoles/Source/VPRolesEditor/Public/VPRolesEditorModule.h
new file mode 100644
--- /dev/null
+++ b/Engine/Plugins/Experimental/VirtualProduction/VPRoles/Source/VPRolesEditor/Public/VPRolesEditorModule.h
@@ -0,0 +1,53 @@
+// Copyright Epic Games, Inc. All Rights Reserved.
+
+#pragma once
+
+#include "CoreMinimal.h"
+#include "Modules/ModuleInterface.h"
+#include "Types/SlateEnums.h"
+
+class SWidget;
+struct EVisibility;
+
+DECLARE_LOG_CATEGORY_EXTERN(LogVPRolesEditor, Log, Log);
+
+class VPROLESEDITOR_API FVPRolesEditorModule : public IModuleInterface
+{
+public:
+	//~ Begin IModuleInterface
+	virtual void StartupModule() override;
+	virtual void ShutdownModule() override;
+	//~End IModuleInterface
+
+private:
+	/** Core delegate handler for extending the level editor toolbar. */
+	void ExtendLevelEditorToolbar();
+	
+	/** Generates the toolbar button widget. */
+	TSharedRef<SWidget> GenerateVPRolesLevelEditorToolbarMenu();
+
+	/** Generate the gameplay tag widget. */
+	TSharedRef<SWidget> GenerateGameplayTagWidget();
+	
+	/** Callback for when a new tag is added */
+	void OnSubmitNewTag(const FText& CommittedText, ETextCommit::Type CommitType);
+
+	/** Callback for when the add tag textbox's text changed. */
+	void OnAddTagTextChanged(const FText& Text);
+	
+	/** Whether the file watcher widget should be shown for the default tag source config file. */
+	EVisibility DetermineFileWatcherWidgetVisibility() const;
+	
+	/** Whether the default VP role tag source can be modified. */
+	bool CanModifyTagSource() const;
+
+private:
+	/** The file name for the default vp role tag source. */
+	FString DefaultTagSourceConfigName;
+	/** Holds the widget generated by the GameplayTags editor module. */
+	TWeakPtr<class SGameplayTagWidget> GameplayTagsWidget;
+	/** Holds the editable text box widget, used to update the text box's error message. */
+	TWeakPtr<class SEditableTextBox> AddRoleEditableBoxWeakPtr;
+	/** Whether to show the add role textbox. */
+	bool bShowAddRoleTextBox = false;
+};
\ No newline at end of file
diff --git a/Engine/Plugins/Experimental/VirtualProduction/VPRoles/Source/VPRolesEditor/VPRolesEditor.Build.cs b/Engine/Plugins/Experimental/VirtualProduction/VPRoles/Source/VPRolesEditor/VPRolesEditor.Build.cs
new file mode 100644
--- /dev/null
+++ b/Engine/Plugins/Experimental/VirtualProduction/VPRoles/Source/VPRolesEditor/VPRolesEditor.Build.cs
@@ -0,0 +1,33 @@
+// Copyright Epic Games, Inc. All Rights Reserved.
+
+using UnrealBuildTool;
+
+public class VPRolesEditor : ModuleRules
+{
+	public VPRolesEditor(ReadOnlyTargetRules Target) : base(Target)
+	{
+		PublicDependencyModuleNames.AddRange(
+			new string[]
+			{
+				"Core",
+			}
+		);
+
+		PrivateDependencyModuleNames.AddRange(
+			new string[]
+			{
+				"CoreUObject",
+				"Engine",
+				"GameplayTags",
+				"GameplayTagsEditor",
+				"LevelEditor",
+				"Projects",
+				"Slate",
+				"SlateCore",
+				"ToolMenus",
+				"VPRoles",
+				"SharedSettingsWidgets"
+			}
+		);
+	}
+}
diff --git a/Engine/Plugins/Experimental/VirtualProduction/VPRoles/VPRoles.uplugin b/Engine/Plugins/Experimental/VirtualProduction/VPRoles/VPRoles.uplugin
new file mode 100644
--- /dev/null
+++ b/Engine/Plugins/Experimental/VirtualProduction/VPRoles/VPRoles.uplugin
@@ -0,0 +1,39 @@
+{
+	"FileVersion": 1,
+	"Version": 1,
+	"VersionName": "1.0",
+	"FriendlyName": "Virtual Production Roles",
+	"Description": "Allows users to manage Virtual Production Role assignment.",
+	"Category": "Virtual Production",
+	"CreatedBy": "Epic Games, Inc.",
+	"CreatedByURL": "https://epicgames.com",
+	"DocsURL": "",
+	"MarketplaceURL": "",
+	"SupportURL": "",
+	"CanContainContent": true,
+	"Hidden" : true,
+	"IsBetaVersion": true,
+	"Installed": false,
+	"Modules": [
+		{
+			"Name": "VPRoles",
+			"Type": "Runtime",
+			"LoadingPhase": "Default"
+		},
+		{
+			"Name": "VPRolesEditor",
+			"Type": "Editor",
+			"LoadingPhase": "Default"
+		},
+	],
+	"Plugins": [
+		{
+			"Name": "VPSettings",
+			"Enabled": true
+		},
+		{
+			"Name": "GameplayTagsEditor",
+			"Enabled": true
+		}
+	]
+}
diff --git a/Engine/Plugins/Experimental/VirtualProduction/VPSettings/Source/Private/VPSettings.cpp b/Engine/Plugins/Experimental/VirtualProduction/VPSettings/Source/Private/VPSettings.cpp
new file mode 100644
--- /dev/null
+++ b/Engine/Plugins/Experimental/VirtualProduction/VPSettings/Source/Private/VPSettings.cpp
@@ -0,0 +1,19 @@
+// Copyright Epic Games, Inc. All Rights Reserved.
+
+#include "VPSettings.h"
+
+const FGameplayTagContainer& UVPSettings::GetRoles() const
+{
+	static const FGameplayTagContainer EmptyContainer;
+	if (ensure(RolesGetter.IsBound()))
+	{
+		return RolesGetter.Execute();
+	}
+	
+	return EmptyContainer;
+}
+
+UVPSettings* UVPSettings::GetVPSettings()
+{
+	return GetMutableDefault<UVPSettings>();
+}
diff --git a/Engine/Plugins/Experimental/VirtualProduction/VPSettings/Source/Private/VPSettingsModule.cpp b/Engine/Plugins/Experimental/VirtualProduction/VPSettings/Source/Private/VPSettingsModule.cpp
new file mode 100644
--- /dev/null
+++ b/Engine/Plugins/Experimental/VirtualProduction/VPSettings/Source/Private/VPSettingsModule.cpp
@@ -0,0 +1,6 @@
+// Copyright Epic Games, Inc. All Rights Reserved.
+
+#include "Modules/ModuleManager.h"
+#include "Modules/ModuleInterface.h"
+
+IMPLEMENT_MODULE(FDefaultModuleImpl, VPSettings)
diff --git a/Engine/Plugins/Experimental/VirtualProduction/VPSettings/Source/Public/VPSettings.h b/Engine/Plugins/Experimental/VirtualProduction/VPSettings/Source/Public/VPSettings.h
new file mode 100644
--- /dev/null
+++ b/Engine/Plugins/Experimental/VirtualProduction/VPSettings/Source/Public/VPSettings.h
@@ -0,0 +1,75 @@
+// Copyright Epic Games, Inc. All Rights Reserved.
+
+#pragma once
+
+#include "GameplayTagContainer.h"
+#include "UObject/Object.h"
+#include "Containers/Array.h"
+#include "VPSettings.generated.h"
+
+/**
+ *
+ */
+UCLASS(config=Game)
+class VPSETTINGS_API UVPSettings : public UObject
+{
+	GENERATED_BODY()
+
+public:
+	UE_DEPRECATED(5.1, "Use UVirtualProductionRolesSubsystem::GetCurrentVirtualProductionRoles instead.")
+	const FGameplayTagContainer& GetRoles() const;
+
+	/** Default Kit of Focal Lengths for Virtual Camera */
+	UPROPERTY(EditAnywhere, BlueprintReadWrite, config, Category = "VirtualCamera|Presets")
+	TArray<float> FocalLengthPresets = { 18.0,21.0,25.0,32.0,40.0,50.0,65.0,75.0,100.0,135.0 };
+
+	/** Default Apertures for for Virtual Camera  */
+	UPROPERTY(EditAnywhere, BlueprintReadWrite, config, Category = "VirtualCamera|Presets")
+	TArray<float> AperturePresets = { 1.0,1.4,2.0,2.8,4.0,5.6,8.0,11.0,16.0,22.0 };
+
+	/** Default Shutter Speeds (1/s) for Virtual Camera*/
+	UPROPERTY(EditAnywhere, BlueprintReadWrite, config, Category = "VirtualCamera|Presets")
+	TArray<float> DefaultShutterSpeedPresets = { 1.0,4.0,8.0,15.0,30.0,60.0,125.0,250.0,500.0,1000.0 };
+
+	/** Default ISOs for Virtual Camera*/
+	UPROPERTY(EditAnywhere, BlueprintReadWrite, config, Category = "VirtualCamera|Presets")
+	TArray<float> DefaultISOPresets = { 50,100,200,400,800,1600,3200,6400 };
+
+	UFUNCTION(BlueprintPure, Category = "Virtual Production")
+	static UVPSettings* GetVPSettings();
+
+#if WITH_EDITORONLY_DATA
+
+	/** When enabled, the virtual production role(s) will be displayed in the main editor UI. */
+	UPROPERTY(config, EditAnywhere, Category="Virtual Production")
+	bool bShowRoleInEditor;
+
+	UE_DEPRECATED(5.1, "Use UVirtualProductionRolesSubsystem::OnRolesChanged instead.")
+	/** Notify when the virtual production roles have changed. */
+	FSimpleMulticastDelegate OnRolesChanged;
+
+	UPROPERTY(config, EditAnywhere, Category = "Virtual Production", DisplayName = "Director Name")
+	FString DirectorName;
+
+	UPROPERTY(config, EditAnywhere, Category = "Virtual Production", DisplayName = "Project Name")
+	FString ShowName;
+#endif
+	
+protected:
+	/**
+	 * The machine role(s) in a virtual production context.
+	 * @note The role may be override via the command line, "-VPRole=[Role.SubRole1|Role.SubRole2]"
+	 */
+	UPROPERTY(config)
+	FGameplayTagContainer Roles;
+
+private:
+	/**
+	 * Roles getter set by the VP Roles module to maintain backwards compatibility.
+	 * May return the command line roles or roles enabled through the VP Roles subsystem.
+	 */
+	DECLARE_DELEGATE_RetVal(const FGameplayTagContainer&, FOnGetRoles);
+	FOnGetRoles RolesGetter;
+
+	friend class UVirtualProductionRolesSubsystem;
+};
diff --git a/Engine/Plugins/Experimental/VirtualProduction/VPSettings/Source/VPSettings.Build.cs b/Engine/Plugins/Experimental/VirtualProduction/VPSettings/Source/VPSettings.Build.cs
new file mode 100644
--- /dev/null
+++ b/Engine/Plugins/Experimental/VirtualProduction/VPSettings/Source/VPSettings.Build.cs
@@ -0,0 +1,19 @@
+// Copyright Epic Games, Inc. All Rights Reserved.
+
+using UnrealBuildTool;
+
+public class VPSettings : ModuleRules
+{
+	public VPSettings(ReadOnlyTargetRules Target) : base(Target)
+	{
+		PublicDependencyModuleNames.AddRange(
+			new string[]
+			{
+				"Core",
+				"Engine",
+				"CoreUObject",
+				"GameplayTags",
+			}
+		);
+	}
+}
diff --git a/Engine/Plugins/Experimental/VirtualProduction/VPSettings/VPSettings.uplugin b/Engine/Plugins/Experimental/VirtualProduction/VPSettings/VPSettings.uplugin
new file mode 100644
--- /dev/null
+++ b/Engine/Plugins/Experimental/VirtualProduction/VPSettings/VPSettings.uplugin
@@ -0,0 +1,24 @@
+{
+	"FileVersion": 1,
+	"Version": 1,
+	"VersionName": "1.0",
+	"FriendlyName": "Virtual Production Settings",
+	"Description": "Allows users to manage Virtual Production settings.",
+	"Category": "Virtual Production",
+	"CreatedBy": "Epic Games, Inc.",
+	"CreatedByURL": "https://epicgames.com",
+	"DocsURL": "",
+	"MarketplaceURL": "",
+	"SupportURL": "",
+	"CanContainContent": true,
+	"Hidden" : true,
+	"IsBetaVersion": true,
+	"Installed": false,
+	"Modules": [
+		{
+			"Name": "VPSettings",
+			"Type": "Runtime",
+			"LoadingPhase": "Default"
+		},
+	]
+}
diff --git a/Engine/Plugins/Experimental/VirtualProductionUtilities/Config/BaseVirtualProductionUtilities.ini b/Engine/Plugins/Experimental/VirtualProductionUtilities/Config/BaseVirtualProductionUtilities.ini
--- a/Engine/Plugins/Experimental/VirtualProductionUtilities/Config/BaseVirtualProductionUtilities.ini
+++ b/Engine/Plugins/Experimental/VirtualProductionUtilities/Config/BaseVirtualProductionUtilities.ini
@@ -1,6 +1,8 @@
 [CoreRedirects]
 
 +ClassRedirects=(OldName="VRTool",NewName="/Script/VPUtilitiesEditor.VRTool")
++ClassRedirects=(OldName="VPSettings",NewName="/Script/VPSettings.VPSettings")
++ClassRedirects=(OldName="VirtualProductionRolesSubsystem",NewName="/Script/VPRoles.VirtualProductionRolesSubsystem")
 +FunctionRedirects=(OldName="VPViewportTickableActorBase.LockLocation",NewName="VPViewportTickableActorBase.EditorLockLocation")
 
 [/Script/VPUtilitiesEditor.VPUtilitiesEditorSettings]
diff --git a/Engine/Plugins/Experimental/VirtualProductionUtilities/Source/VPRoles/Private/VPRolesModule.cpp b/Engine/Plugins/Experimental/VirtualProductionUtilities/Source/VPRoles/Private/VPRolesModule.cpp
deleted file mode 100644
diff --git a/Engine/Plugins/Experimental/VirtualProductionUtilities/Source/VPRoles/Private/VPRolesModule.h b/Engine/Plugins/Experimental/VirtualProductionUtilities/Source/VPRoles/Private/VPRolesModule.h
deleted file mode 100644
diff --git a/Engine/Plugins/Experimental/VirtualProductionUtilities/Source/VPRoles/Private/VPRolesSettings.h b/Engine/Plugins/Experimental/VirtualProductionUtilities/Source/VPRoles/Private/VPRolesSettings.h
deleted file mode 100644
diff --git a/Engine/Plugins/Experimental/VirtualProductionUtilities/Source/VPRoles/Private/VPRolesSubsystem.cpp b/Engine/Plugins/Experimental/VirtualProductionUtilities/Source/VPRoles/Private/VPRolesSubsystem.cpp
deleted file mode 100644
diff --git a/Engine/Plugins/Experimental/VirtualProductionUtilities/Source/VPRoles/Public/VPRolesSubsystem.h b/Engine/Plugins/Experimental/VirtualProductionUtilities/Source/VPRoles/Public/VPRolesSubsystem.h
deleted file mode 100644
diff --git a/Engine/Plugins/Experimental/VirtualProductionUtilities/Source/VPRoles/VPRoles.Build.cs b/Engine/Plugins/Experimental/VirtualProductionUtilities/Source/VPRoles/VPRoles.Build.cs
deleted file mode 100644
diff --git a/Engine/Plugins/Experimental/VirtualProductionUtilities/Source/VPRolesEditor/Private/VPRolesEditorModule.cpp b/Engine/Plugins/Experimental/VirtualProductionUtilities/Source/VPRolesEditor/Private/VPRolesEditorModule.cpp
deleted file mode 100644
diff --git a/Engine/Plugins/Experimental/VirtualProductionUtilities/Source/VPRolesEditor/Private/VPRolesEditorStyle.cpp b/Engine/Plugins/Experimental/VirtualProductionUtilities/Source/VPRolesEditor/Private/VPRolesEditorStyle.cpp
deleted file mode 100644
diff --git a/Engine/Plugins/Experimental/VirtualProductionUtilities/Source/VPRolesEditor/Private/VPRolesEditorStyle.h b/Engine/Plugins/Experimental/VirtualProductionUtilities/Source/VPRolesEditor/Private/VPRolesEditorStyle.h
deleted file mode 100644
diff --git a/Engine/Plugins/Experimental/VirtualProductionUtilities/Source/VPRolesEditor/Public/VPRolesEditorModule.h b/Engine/Plugins/Experimental/VirtualProductionUtilities/Source/VPRolesEditor/Public/VPRolesEditorModule.h
deleted file mode 100644
diff --git a/Engine/Plugins/Experimental/VirtualProductionUtilities/Source/VPRolesEditor/VPRolesEditor.Build.cs b/Engine/Plugins/Experimental/VirtualProductionUtilities/Source/VPRolesEditor/VPRolesEditor.Build.cs
deleted file mode 100644
diff --git a/Engine/Plugins/Experimental/VirtualProductionUtilities/Source/VPUtilities/Private/VPBookmarkActor.cpp b/Engine/Plugins/Experimental/VirtualProductionUtilities/Source/VPUtilities/Private/VPBookmarkActor.cpp
--- a/Engine/Plugins/Experimental/VirtualProductionUtilities/Source/VPUtilities/Private/VPBookmarkActor.cpp
+++ b/Engine/Plugins/Experimental/VirtualProductionUtilities/Source/VPUtilities/Private/VPBookmarkActor.cpp
@@ -12,7 +12,7 @@
 #include "VPBlueprintLibrary.h"
 #include "VPBookmark.h"
 #include "Materials/MaterialInstanceDynamic.h"
-#include "VPSettings.h"
+#include "VPBookmarkSettings.h"
 
 AVPBookmarkActor::AVPBookmarkActor(const FObjectInitializer& ObjectInitializer)
 	: Super(ObjectInitializer)
@@ -250,4 +250,4 @@ void AVPBookmarkActor::OnConstruction(const FTransform& Transform)
 	Super::OnConstruction(Transform);
 
 	AVPBookmarkActor::UpdateBookmarkColor(BookmarkColor);
-}
\ No newline at end of file
+}
diff --git a/Engine/Plugins/Experimental/VirtualProductionUtilities/Source/VPUtilities/Private/VPSettings.cpp b/Engine/Plugins/Experimental/VirtualProductionUtilities/Source/VPUtilities/Private/VPSettings.cpp
deleted file mode 100644
diff --git a/Engine/Plugins/Experimental/VirtualProductionUtilities/Source/VPUtilities/Public/VPBookmarkSettings.h b/Engine/Plugins/Experimental/VirtualProductionUtilities/Source/VPUtilities/Public/VPBookmarkSettings.h
new file mode 100644
--- /dev/null
+++ b/Engine/Plugins/Experimental/VirtualProductionUtilities/Source/VPUtilities/Public/VPBookmarkSettings.h
@@ -0,0 +1,33 @@
+// Copyright Epic Games, Inc. All Rights Reserved.
+
+#pragma once
+
+#include "UObject/SoftObjectPath.h"
+#include "VPBookmarkSettings.generated.h"
+
+UCLASS(Config = VirtualProductionUtilities)
+class VPUTILITIES_API UVPBookmarkSettings : public UObject
+{
+	GENERATED_BODY()
+		//~ Bookmark assets //
+public:
+	/** Bookmark mesh to use for Bookmark actor */
+	UPROPERTY(config, EditAnywhere, Category = "VP Bookmark", meta = (AllowedClasses = "/Script/Engine.StaticMesh", DisplayName = "Bookmark Mesh"))
+	FSoftObjectPath BookmarkMeshPath;
+
+	/** Bookmark mesh material to use for Bookmark actor */
+	UPROPERTY(config, EditAnywhere, Category = "VP Bookmark", meta = (AllowedClasses = "/Script/Engine.MaterialInterface", DisplayName = "Bookmark Mesh Material"))
+	FSoftObjectPath BookmarkMaterialPath;
+
+	/** Bookmark spline mesh to use for Bookmark actor */
+	UPROPERTY(config, EditAnywhere, Category = "VP Bookmark", meta = (AllowedClasses = "/Script/Engine.StaticMesh", DisplayName = "Bookmark Spline Mesh"))
+	FSoftObjectPath BookmarkSplineMeshPath;
+
+	/** Bookmark spline mesh material to use for Bookmark actor */
+	UPROPERTY(config, EditAnywhere, Category = "VP Bookmark", meta = (AllowedClasses = "/Script/Engine.MaterialInterface", DisplayName = "Bookmark Spline Material"))
+	FSoftObjectPath BookmarkSplineMeshMaterialPath;
+
+	/** Bookmark label material class to use for Bookmark actor */
+	UPROPERTY(config, EditAnywhere, Category = "VP Bookmark", meta = (AllowedClasses = "/Script/Engine.MaterialInterface", DisplayName = "Bookmark Label Material"))
+	FSoftObjectPath BookmarkLabelMaterialPath;
+};
diff --git a/Engine/Plugins/Experimental/VirtualProductionUtilities/Source/VPUtilities/Public/VPSettings.h b/Engine/Plugins/Experimental/VirtualProductionUtilities/Source/VPUtilities/Public/VPSettings.h
deleted file mode 100644
diff --git a/Engine/Plugins/Experimental/VirtualProductionUtilities/Source/VPUtilities/VPUtilities.Build.cs b/Engine/Plugins/Experimental/VirtualProductionUtilities/Source/VPUtilities/VPUtilities.Build.cs
--- a/Engine/Plugins/Experimental/VirtualProductionUtilities/Source/VPUtilities/VPUtilities.Build.cs
+++ b/Engine/Plugins/Experimental/VirtualProductionUtilities/Source/VPUtilities/VPUtilities.Build.cs
@@ -20,6 +20,7 @@ public class VPUtilities : ModuleRules
 				"TimeManagement",
 				"UMG",
 				"VPBookmark",
+				"VPSettings"
 			}
 		);
 
diff --git a/Engine/Plugins/Experimental/VirtualProductionUtilities/Source/VPUtilitiesEditor/VPUtilitiesEditor.Build.cs b/Engine/Plugins/Experimental/VirtualProductionUtilities/Source/VPUtilitiesEditor/VPUtilitiesEditor.Build.cs
--- a/Engine/Plugins/Experimental/VirtualProductionUtilities/Source/VPUtilitiesEditor/VPUtilitiesEditor.Build.cs
+++ b/Engine/Plugins/Experimental/VirtualProductionUtilities/Source/VPUtilitiesEditor/VPUtilitiesEditor.Build.cs
@@ -41,6 +41,7 @@ public class VPUtilitiesEditor : ModuleRules
 				"UnrealEd",
 				"VPBookmark",
 				"VPRoles",
+				"VPSettings",
 				"WorkspaceMenuStructure",
             }
 		);
diff --git a/Engine/Plugins/Experimental/VirtualProductionUtilities/VirtualProductionUtilities.uplugin b/Engine/Plugins/Experimental/VirtualProductionUtilities/VirtualProductionUtilities.uplugin
--- a/Engine/Plugins/Experimental/VirtualProductionUtilities/VirtualProductionUtilities.uplugin
+++ b/Engine/Plugins/Experimental/VirtualProductionUtilities/VirtualProductionUtilities.uplugin
@@ -39,16 +39,6 @@
 			"Type": "Editor",
 			"LoadingPhase": "PostDefault"
 		},
-		{
-			"Name": "VPRoles",
-			"Type": "Runtime",
-			"LoadingPhase": "Default"
-		},
-		{
-			"Name": "VPRolesEditor",
-			"Type": "Editor",
-			"LoadingPhase": "Default"
-		},
 		{
 			"Name": "VPSpline",
 			"Type": "Runtime",
@@ -86,7 +76,11 @@
 			"Enabled": true
 		},
 		{
-			"Name": "GameplayTagsEditor",
+			"Name": "VPRoles",
+			"Enabled": true
+		},
+		{
+			"Name": "VPSettings",
 			"Enabled": true
 		},
 		{
