Fixed non unity compile errors caused by integration from RES. Errors were reported by farm

diff --git a/Engine/Plugins/Experimental/PCG/Source/PCG/Private/Elements/PCGStaticMeshSpawner.cpp b/Engine/Plugins/Experimental/PCG/Source/PCG/Private/Elements/PCGStaticMeshSpawner.cpp
--- a/Engine/Plugins/Experimental/PCG/Source/PCG/Private/Elements/PCGStaticMeshSpawner.cpp
+++ b/Engine/Plugins/Experimental/PCG/Source/PCG/Private/Elements/PCGStaticMeshSpawner.cpp
@@ -339,3 +339,5 @@ void UPCGStaticMeshSpawnerSettings::RefreshInstancePacker()
 	}
 }
 
+FPCGStaticMeshSpawnerContext::FPackedInstanceListData::FPackedInstanceListData() = default;
+FPCGStaticMeshSpawnerContext::FPackedInstanceListData::~FPackedInstanceListData() = default;
diff --git a/Engine/Plugins/Experimental/PCG/Source/PCG/Public/Elements/PCGStaticMeshSpawner.h b/Engine/Plugins/Experimental/PCG/Source/PCG/Public/Elements/PCGStaticMeshSpawner.h
--- a/Engine/Plugins/Experimental/PCG/Source/PCG/Public/Elements/PCGStaticMeshSpawner.h
+++ b/Engine/Plugins/Experimental/PCG/Source/PCG/Public/Elements/PCGStaticMeshSpawner.h
@@ -104,6 +104,8 @@ struct FPCGStaticMeshSpawnerContext : public FPCGContext
 
 	struct FPackedInstanceListData
 	{
+		FPackedInstanceListData();
+		~FPackedInstanceListData();
 		const UPCGSpatialData* SpatialData;
 		TArray<FPCGMeshInstanceList> MeshInstances;
 		TArray<FPCGPackedCustomData> PackedCustomData;
diff --git a/Engine/Plugins/Experimental/VirtualCamera/Source/VCamCore/Private/Output/VCamOutputRemoteSession.cpp b/Engine/Plugins/Experimental/VirtualCamera/Source/VCamCore/Private/Output/VCamOutputRemoteSession.cpp
--- a/Engine/Plugins/Experimental/VirtualCamera/Source/VCamCore/Private/Output/VCamOutputRemoteSession.cpp
+++ b/Engine/Plugins/Experimental/VirtualCamera/Source/VCamCore/Private/Output/VCamOutputRemoteSession.cpp
@@ -22,6 +22,7 @@
 #include "SceneView.h"
 #else
 #include "RemoteSession.h"
+#include "UObject/Package.h"
 #endif
 
 namespace UE::VCamOutputRemoteSession::Private
diff --git a/Engine/Plugins/Experimental/VirtualProduction/VPRoles/Source/VPRolesEditor/Private/VPRolesEditorModule.cpp b/Engine/Plugins/Experimental/VirtualProduction/VPRoles/Source/VPRolesEditor/Private/VPRolesEditorModule.cpp
--- a/Engine/Plugins/Experimental/VirtualProduction/VPRoles/Source/VPRolesEditor/Private/VPRolesEditorModule.cpp
+++ b/Engine/Plugins/Experimental/VirtualProduction/VPRoles/Source/VPRolesEditor/Private/VPRolesEditorModule.cpp
@@ -15,6 +15,7 @@
 #include "ToolMenuSection.h"
 #include "VPRolesEditorStyle.h"
 #include "VPRolesSubsystem.h"
+#include "Widgets/Images/SImage.h"
 #include "Widgets/Input/SCheckBox.h"
 #include "Widgets/SNullWidget.h"
 #include "Widgets/Input/SEditableTextBox.h"
